Work on Localization: Localizable Message Module, implements a message data structure that supports replication and resolve localization on the client.

diff --git a/Engine/Plugins/Experimental/LocalizableMessage/LocalizableMessage.uplugin b/Engine/Plugins/Experimental/LocalizableMessage/LocalizableMessage.uplugin
new file mode 100644
--- /dev/null
+++ b/Engine/Plugins/Experimental/LocalizableMessage/LocalizableMessage.uplugin
@@ -0,0 +1,25 @@
+{
+	"FileVersion" : 3,
+	"Version" : 1,
+	"VersionName" : "1.0",
+	"FriendlyName" : "Localizable Message",
+	"Description" : "Utility for a text message that can be replicated. It supports parameter resolution for the client.",
+	"Category" : "Messaging",
+	"CreatedBy" : "Epic Games, Inc.",
+	"CreatedByURL" : "https://epicgames.com",
+	"DocsURL" : "",
+	"MarketplaceURL" : "",
+	"SupportURL" : "",
+	"EnabledByDefault" : true,
+	"CanContainContent" : false,
+	"IsBetaVersion" : false,
+	"Installed" : false,
+	"Modules" :
+	[
+		{
+			"Name": "LocalizableMessage",
+			"Type": "Runtime",
+			"LoadingPhase": "Default"
+		}
+	]
+}
diff --git a/Engine/Plugins/Experimental/LocalizableMessage/Source/LocalizableMessage/LocalizableMessage.Build.cs b/Engine/Plugins/Experimental/LocalizableMessage/Source/LocalizableMessage/LocalizableMessage.Build.cs
new file mode 100644
--- /dev/null
+++ b/Engine/Plugins/Experimental/LocalizableMessage/Source/LocalizableMessage/LocalizableMessage.Build.cs
@@ -0,0 +1,16 @@
+// Copyright Epic Games, Inc. All Rights Reserved.
+
+namespace UnrealBuildTool.Rules
+{
+	public class LocalizableMessage : ModuleRules
+	{
+		public LocalizableMessage(ReadOnlyTargetRules Target) : base(Target)
+		{
+			PrivateDependencyModuleNames.AddRange(
+				new string[] {
+					"Core",
+					"CoreUObject",
+				});
+		}
+	}
+}
diff --git a/Engine/Plugins/Experimental/LocalizableMessage/Source/LocalizableMessage/Private/LocalizableMessage.cpp b/Engine/Plugins/Experimental/LocalizableMessage/Source/LocalizableMessage/Private/LocalizableMessage.cpp
new file mode 100644
--- /dev/null
+++ b/Engine/Plugins/Experimental/LocalizableMessage/Source/LocalizableMessage/Private/LocalizableMessage.cpp
@@ -0,0 +1,96 @@
+// Copyright Epic Games, Inc. All Rights Reserved.
+
+#include "LocalizableMessage.h"
+
+FLocalizableMessageParameterEntry::FLocalizableMessageParameterEntry() = default;
+FLocalizableMessageParameterEntry::FLocalizableMessageParameterEntry(const FString& InKey, TUniquePtr<FLocalizableMessageParameter, FLocalizableMessageParameter::FCustomDeleter>&& InValue) : Key(InKey), Value(std::move(InValue)) {}
+FLocalizableMessageParameterEntry::FLocalizableMessageParameterEntry(FLocalizableMessageParameterEntry&&) = default;
+FLocalizableMessageParameterEntry::~FLocalizableMessageParameterEntry() = default;
+FLocalizableMessageParameterEntry& FLocalizableMessageParameterEntry::operator=(FLocalizableMessageParameterEntry&&) = default;
+
+bool FLocalizableMessageParameterEntry::operator==(const FLocalizableMessageParameterEntry& Other) const
+{
+	if (Key != Other.Key)
+	{
+		return false;
+	}
+
+	if (Value.IsValid() != Other.Value.IsValid())
+	{
+		return false;
+	}
+
+	if (Value.IsValid())
+	{
+		const UScriptStruct* TypeA = Value->GetScriptStruct();
+		const UScriptStruct* TypeB = Other.Value->GetScriptStruct();
+
+		if (TypeA != TypeB)
+		{
+			return false;
+		}
+
+		if (TypeA->CompareScriptStruct(Value.Get(), Other.Value.Get(), 0) == false)
+		{
+			return false;
+		}
+	}
+
+	return true;
+}
+
+bool FLocalizableMessageParameterEntry::NetSerialize(FArchive& Ar, UPackageMap* Map, bool& bOutSuccess)
+{
+	TCheckedObjPtr<UScriptStruct> EntryType = Value.IsValid() ? Value->GetScriptStruct() : nullptr;
+
+	Ar << Key;
+	Ar << EntryType;
+
+	if (Ar.IsLoading())
+	{
+		Value = nullptr;
+
+		if (EntryType.IsValid())
+		{
+			FLocalizableMessageParameter* NewEntry = FLocalizableMessageParameter::AllocateType(EntryType.Get());
+			Value = TUniquePtr<FLocalizableMessageParameter, FLocalizableMessageParameter::FCustomDeleter>(NewEntry);
+		}
+	}
+
+	if (Value.IsValid())
+	{
+		if (EntryType->StructFlags & STRUCT_NetSerializeNative)
+		{
+			EntryType->GetCppStructOps()->NetSerialize(Ar, Map, bOutSuccess, Value.Get());
+		}
+		else
+		{
+			EntryType->SerializeItem(FStructuredArchiveFromArchive(Ar).GetSlot(), Value.Get(), NULL);
+		}
+	}
+
+	bOutSuccess = true;
+
+	return true;
+}
+
+void FLocalizableMessageParameterEntry::AddStructReferencedObjects(class FReferenceCollector& Collector)
+{
+	if (FLocalizableMessageParameter* Param = Value.Get())
+	{
+		Collector.AddPropertyReferencesWithStructARO(Param->GetScriptStruct(), Param);
+	}
+}
+
+FLocalizableMessage::FLocalizableMessage() = default;
+FLocalizableMessage::FLocalizableMessage(FLocalizableMessage&&) = default;
+FLocalizableMessage::~FLocalizableMessage() = default;
+FLocalizableMessage& FLocalizableMessage::operator=(FLocalizableMessage&&) = default;
+
+
+bool FLocalizableMessage::operator==(const FLocalizableMessage& Other) const
+{
+	return Key == Other.Key &&
+		DefaultText == Other.DefaultText &&
+		Substitutions == Other.Substitutions;
+}
diff --git a/Engine/Plugins/Experimental/LocalizableMessage/Source/LocalizableMessage/Private/LocalizableMessageModule.cpp b/Engine/Plugins/Experimental/LocalizableMessage/Source/LocalizableMessage/Private/LocalizableMessageModule.cpp
new file mode 100644
--- /dev/null
+++ b/Engine/Plugins/Experimental/LocalizableMessage/Source/LocalizableMessage/Private/LocalizableMessageModule.cpp
@@ -0,0 +1,42 @@
+// Copyright Epic Games, Inc. All Rights Reserved.
+
+#include "ILocalizableMessageModule.h"
+#include "LocalizableMessageProcessor.h"
+#include "LocalizableMessageTypes.h"
+#include "Modules/ModuleManager.h"
+
+class FLocalizableMessageModule : public ILocalizableMessageModule
+{
+public:
+
+	/** IModuleInterface implementation */
+	virtual void StartupModule() override;
+	virtual void ShutdownModule() override;
+
+	virtual FLocalizableMessageProcessor& GetLocalizableMessageProcessor() override { return LocalizableMessageProcessor; }
+
+private:
+
+	FLocalizableMessageProcessor LocalizableMessageProcessor;
+};
+
+void FLocalizableMessageModule::StartupModule()
+{
+	LocalizableMessageTypes::RegisterTypes();
+}
+
+void FLocalizableMessageModule::ShutdownModule()
+{
+	LocalizableMessageTypes::UnregisterTypes();
+}
+
+ILocalizableMessageModule& ILocalizableMessageModule::Get()
+{
+	ILocalizableMessageModule* VerseModule;
+	UE_AUTORTFM_OPEN(
+		VerseModule = &FModuleManager::LoadModuleChecked<FLocalizableMessageModule>(TEXT("LocalizableMessage"));
+	);
+	return *VerseModule;
+}
+
+IMPLEMENT_MODULE(FLocalizableMessageModule, LocalizableMessage)
diff --git a/Engine/Plugins/Experimental/LocalizableMessage/Source/LocalizableMessage/Private/LocalizableMessageParameter.cpp b/Engine/Plugins/Experimental/LocalizableMessage/Source/LocalizableMessage/Private/LocalizableMessageParameter.cpp
new file mode 100644
--- /dev/null
+++ b/Engine/Plugins/Experimental/LocalizableMessage/Source/LocalizableMessage/Private/LocalizableMessageParameter.cpp
@@ -0,0 +1,21 @@
+// Copyright Epic Games, Inc. All Rights Reserved.
+
+#include "LocalizableMessageParameter.h"
+
+FLocalizableMessageParameter* FLocalizableMessageParameter::AllocateType(const UScriptStruct* Type)
+{
+	FLocalizableMessageParameter* NewEntry = (FLocalizableMessageParameter*)FMemory::Malloc(Type->GetCppStructOps()->GetSize());
+	Type->InitializeStruct(NewEntry);
+	return NewEntry;
+}
+
+void FLocalizableMessageParameter::FCustomDeleter::operator()(FLocalizableMessageParameter* Object) const
+{
+	if (Object)
+	{
+		UScriptStruct* ScriptStruct = Object->GetScriptStruct();
+		check(ScriptStruct);
+		ScriptStruct->DestroyStruct(Object);
+		FMemory::Free(Object);
+	}
+}
diff --git a/Engine/Plugins/Experimental/LocalizableMessage/Source/LocalizableMessage/Private/LocalizableMessageProcessor.cpp b/Engine/Plugins/Experimental/LocalizableMessage/Source/LocalizableMessage/Private/LocalizableMessageProcessor.cpp
new file mode 100644
--- /dev/null
+++ b/Engine/Plugins/Experimental/LocalizableMessage/Source/LocalizableMessage/Private/LocalizableMessageProcessor.cpp
@@ -0,0 +1,63 @@
+// Copyright Epic Games, Inc. All Rights Reserved.
+
+#include "LocalizableMessageProcessor.h"
+
+#include "Internationalization/Text.h"
+#include "LocalizableMessageParameter.h"
+#include "LocalizableMessage.h"
+#include "LocalizationContext.h"
+
+FLocalizableMessageProcessor::FLocalizableMessageProcessor()
+{
+
+}
+
+FLocalizableMessageProcessor::~FLocalizableMessageProcessor()
+{
+	ensure(LocalizeValueMapping.IsEmpty());
+}
+
+FText FLocalizableMessageProcessor::Localize(const FLocalizableMessage& Message, const FLocalizationContext& Context)
+{
+	// todo: look up the localized string via Key
+	// for now we only use the DefaultText in the message
+	FText SubstitutionResult;
+	FFormatNamedArguments FormatArguments;
+
+	for (const FLocalizableMessageParameterEntry& Substitution : Message.Substitutions)
+	{
+		if (const LocalizeValueFnc* Functor = LocalizeValueMapping.Find(Substitution.Value->GetScriptStruct()->GetFName()))
+		{
+			if (ensure(Substitution.Value != nullptr))
+			{
+				SubstitutionResult = (*Functor)(*Substitution.Value, Context);
+
+				if (SubstitutionResult.IsEmpty() == false)
+				{
+					TextFormatUtil::FormatNamed(FormatArguments, Substitution.Key, SubstitutionResult);
+				}
+			}
+			else
+			{
+				ensureMsgf(false, TEXT("Message contained null substitution."));
+			}
+		}
+		else
+		{
+			ensureMsgf(false, TEXT("Localization type %s not registered in Localization Processor."), *Substitution.Value->GetScriptStruct()->GetFName().ToString());
+		}
+	}
+
+	// an unfortunate number of allocations and copies here
+	return FText::Format(FTextFormat::FromString(Message.DefaultText), FormatArguments);
+}
+
+void FLocalizableMessageProcessor::UnregisterLocalizableTypes(FScopedRegistrations& ScopedRegistrations)
+{
+	for (const FName& Registration : ScopedRegistrations.Registrations)
+	{
+		LocalizeValueMapping.Remove(Registration);
+	}
+
+	ScopedRegistrations.Registrations.Reset();
+}
diff --git a/Engine/Plugins/Experimental/LocalizableMessage/Source/LocalizableMessage/Private/LocalizableMessageTypes.cpp b/Engine/Plugins/Experimental/LocalizableMessage/Source/LocalizableMessage/Private/LocalizableMessageTypes.cpp
new file mode 100644
--- /dev/null
+++ b/Engine/Plugins/Experimental/LocalizableMessage/Source/LocalizableMessage/Private/LocalizableMessageTypes.cpp
@@ -0,0 +1,43 @@
+// Copyright Epic Games, Inc. All Rights Reserved.
+
+#include "LocalizableMessageTypes.h"
+
+#include "ILocalizableMessageModule.h"
+#include "LocalizableMessageBaseParameters.h"
+#include "LocalizableMessageProcessor.h"
+
+namespace LocalizableMessageTypes
+{
+	FLocalizableMessageProcessor::FScopedRegistrations RegisteredLocalizationTypes;
+
+	FText Int_LocalizeValue(const FLocalizableMessageParameterInt& Localizable, const FLocalizationContext& LocalizationContext)
+	{
+		return FText::AsNumber(Localizable.Value);
+	}
+	FText Float_LocalizeValue(const FLocalizableMessageParameterFloat& Localizable, const FLocalizationContext& LocalizationContext)
+	{
+		return FText::AsNumber(Localizable.Value);
+	}
+	FText String_LocalizeValue(const FLocalizableMessageParameterString& Localizable, const FLocalizationContext& LocalizationContext)
+	{
+		return FText::FromString(Localizable.Value);
+	}
+
+	void RegisterTypes()
+	{
+		ILocalizableMessageModule& LocalizableMessageModule = ILocalizableMessageModule::Get();
+		FLocalizableMessageProcessor& Processor = LocalizableMessageModule.GetLocalizableMessageProcessor();
+
+		Processor.RegisterLocalizableType<FLocalizableMessageParameterInt>(&Int_LocalizeValue,RegisteredLocalizationTypes);
+		Processor.RegisterLocalizableType<FLocalizableMessageParameterFloat>(&Float_LocalizeValue,RegisteredLocalizationTypes);
+		Processor.RegisterLocalizableType<FLocalizableMessageParameterString>(&String_LocalizeValue, RegisteredLocalizationTypes);
+	}
+
+	void UnregisterTypes()
+	{
+		ILocalizableMessageModule& LocalizableMessageModule = ILocalizableMessageModule::Get();
+		FLocalizableMessageProcessor& Processor = LocalizableMessageModule.GetLocalizableMessageProcessor();
+
+		Processor.UnregisterLocalizableTypes(RegisteredLocalizationTypes);
+	}
+}
\ No newline at end of file
diff --git a/Engine/Plugins/Experimental/LocalizableMessage/Source/LocalizableMessage/Public/ILocalizableMessageModule.h b/Engine/Plugins/Experimental/LocalizableMessage/Source/LocalizableMessage/Public/ILocalizableMessageModule.h
new file mode 100644
--- /dev/null
+++ b/Engine/Plugins/Experimental/LocalizableMessage/Source/LocalizableMessage/Public/ILocalizableMessageModule.h
@@ -0,0 +1,15 @@
+// Copyright Epic Games, Inc. All Rights Reserved.
+
+#pragma once
+
+#include "Modules/ModuleInterface.h"
+
+class FLocalizableMessageProcessor;
+
+class LOCALIZABLEMESSAGE_API ILocalizableMessageModule : public IModuleInterface
+{
+public:
+	static ILocalizableMessageModule& Get();
+
+	virtual FLocalizableMessageProcessor& GetLocalizableMessageProcessor() = 0;
+};
diff --git a/Engine/Plugins/Experimental/LocalizableMessage/Source/LocalizableMessage/Public/LocalizableMessage.h b/Engine/Plugins/Experimental/LocalizableMessage/Source/LocalizableMessage/Public/LocalizableMessage.h
new file mode 100644
--- /dev/null
+++ b/Engine/Plugins/Experimental/LocalizableMessage/Source/LocalizableMessage/Public/LocalizableMessage.h
@@ -0,0 +1,77 @@
+// Copyright Epic Games, Inc. All Rights Reserved.
+
+#pragma once
+
+#include "Containers/Array.h"
+#include "LocalizableMessageParameter.h"
+#include "Templates/UniquePtr.h"
+
+#include "LocalizableMessage.generated.h"
+
+USTRUCT()
+struct FLocalizableMessageParameterEntry
+{
+	GENERATED_BODY()
+
+public:
+
+	FLocalizableMessageParameterEntry();
+	LOCALIZABLEMESSAGE_API FLocalizableMessageParameterEntry(const FString& InKey, TUniquePtr<FLocalizableMessageParameter, FLocalizableMessageParameter::FCustomDeleter>&& InValue);
+	FLocalizableMessageParameterEntry(FLocalizableMessageParameterEntry&&);
+	~FLocalizableMessageParameterEntry();
+	FLocalizableMessageParameterEntry& operator=(FLocalizableMessageParameterEntry&&);
+
+	bool operator==(const FLocalizableMessageParameterEntry& Other) const;
+
+	void AddStructReferencedObjects(class FReferenceCollector& Collector);
+	bool NetSerialize(FArchive& Ar, UPackageMap* Map, bool& bOutSuccess);
+
+	UPROPERTY()
+	FString Key;
+
+	TUniquePtr<FLocalizableMessageParameter, FLocalizableMessageParameter::FCustomDeleter> Value;
+};
+
+template<>
+struct TStructOpsTypeTraits< FLocalizableMessageParameterEntry > : public TStructOpsTypeTraitsBase2< FLocalizableMessageParameterEntry >
+{
+	enum
+	{
+		WithAddStructReferencedObjects = true,
+		WithNetSerializer = true,
+		WithCopy = false,
+	};
+};
+
+
+USTRUCT()
+struct FLocalizableMessage
+{
+	GENERATED_BODY();
+
+	LOCALIZABLEMESSAGE_API FLocalizableMessage();
+	LOCALIZABLEMESSAGE_API FLocalizableMessage(FLocalizableMessage&&);
+	LOCALIZABLEMESSAGE_API ~FLocalizableMessage();
+	LOCALIZABLEMESSAGE_API FLocalizableMessage& operator=(FLocalizableMessage&&);
+
+	LOCALIZABLEMESSAGE_API bool operator==(const FLocalizableMessage& Other) const;
+	LOCALIZABLEMESSAGE_API bool operator!=(const FLocalizableMessage& Other) const { return !(*this == Other); }
+
+	UPROPERTY()
+	FString Key;
+
+	UPROPERTY()
+	FString DefaultText;
+
+	UPROPERTY()
+	TArray<FLocalizableMessageParameterEntry> Substitutions;
+};
+
+template<>
+struct TStructOpsTypeTraits< FLocalizableMessage > : public TStructOpsTypeTraitsBase2< FLocalizableMessage >
+{
+	enum
+	{
+		WithCopy = false,
+	};
+};
\ No newline at end of file
diff --git a/Engine/Plugins/Experimental/LocalizableMessage/Source/LocalizableMessage/Public/LocalizableMessageBaseParameters.h b/Engine/Plugins/Experimental/LocalizableMessage/Source/LocalizableMessage/Public/LocalizableMessageBaseParameters.h
new file mode 100644
--- /dev/null
+++ b/Engine/Plugins/Experimental/LocalizableMessage/Source/LocalizableMessage/Public/LocalizableMessageBaseParameters.h
@@ -0,0 +1,48 @@
+// Copyright Epic Games, Inc. All Rights Reserved.
+
+#pragma once
+
+#include "LocalizableMessageParameter.h"
+
+#include "Containers/UnrealString.h"
+
+#include "LocalizableMessageBaseParameters.generated.h"
+
+USTRUCT()
+struct FLocalizableMessageParameterInt : public FLocalizableMessageParameter
+{
+	GENERATED_BODY()
+
+public:
+
+	UPROPERTY()
+	int64 Value = 0;
+
+	virtual UScriptStruct* GetScriptStruct() const override { return FLocalizableMessageParameterInt::StaticStruct(); }
+};
+
+USTRUCT()
+struct FLocalizableMessageParameterFloat : public FLocalizableMessageParameter
+{
+	GENERATED_BODY()
+
+public:
+
+	UPROPERTY()
+	double Value = 0;
+
+	virtual UScriptStruct* GetScriptStruct() const override { return FLocalizableMessageParameterFloat::StaticStruct(); }
+};
+
+USTRUCT()
+struct FLocalizableMessageParameterString : public FLocalizableMessageParameter
+{
+	GENERATED_BODY()
+
+public:
+
+	UPROPERTY()
+	FString Value;
+
+	virtual UScriptStruct* GetScriptStruct() const override { return FLocalizableMessageParameterString::StaticStruct(); }
+};
diff --git a/Engine/Plugins/Experimental/LocalizableMessage/Source/LocalizableMessage/Public/LocalizableMessageParameter.h b/Engine/Plugins/Experimental/LocalizableMessage/Source/LocalizableMessage/Public/LocalizableMessageParameter.h
new file mode 100644
--- /dev/null
+++ b/Engine/Plugins/Experimental/LocalizableMessage/Source/LocalizableMessage/Public/LocalizableMessageParameter.h
@@ -0,0 +1,31 @@
+// Copyright Epic Games, Inc. All Rights Reserved.
+
+#pragma once
+
+#include "Containers/UnrealString.h"
+
+#include "LocalizableMessageParameter.generated.h"
+
+USTRUCT()
+struct FLocalizableMessageParameter
+{
+	GENERATED_BODY()
+
+	FLocalizableMessageParameter() = default;
+	virtual ~FLocalizableMessageParameter() = default;
+
+	virtual UScriptStruct* GetScriptStruct() const PURE_VIRTUAL(FLocalizableMessageParameter::GetScriptStruct, return FLocalizableMessageParameter::StaticStruct();)
+
+	template <typename ChildType>
+	static ChildType* AllocateType()
+	{
+		return static_cast<ChildType*>(AllocateType(ChildType::StaticStruct()));
+	}
+
+	LOCALIZABLEMESSAGE_API static FLocalizableMessageParameter* AllocateType(const UScriptStruct* Type);
+
+	struct FCustomDeleter
+	{
+		LOCALIZABLEMESSAGE_API void operator()(FLocalizableMessageParameter* Object) const;
+	};
+};
\ No newline at end of file
diff --git a/Engine/Plugins/Experimental/LocalizableMessage/Source/LocalizableMessage/Public/LocalizableMessageProcessor.h b/Engine/Plugins/Experimental/LocalizableMessage/Source/LocalizableMessage/Public/LocalizableMessageProcessor.h
new file mode 100644
--- /dev/null
+++ b/Engine/Plugins/Experimental/LocalizableMessage/Source/LocalizableMessage/Public/LocalizableMessageProcessor.h
@@ -0,0 +1,51 @@
+// Copyright Epic Games, Inc. All Rights Reserved.
+
+#pragma once
+
+#include "Containers/Map.h"
+#include "Templates/Function.h"
+#include "Templates/Tuple.h"
+#include "UObject/NameTypes.h"
+
+class FText;
+struct FLocalizableMessage;
+struct FLocalizableMessageParameter;
+struct FLocalizationContext;
+
+class FLocalizableMessageProcessor
+{
+public:
+
+	struct FScopedRegistrations
+	{
+		FScopedRegistrations() = default;
+		~FScopedRegistrations() { ensure(Registrations.Num() == 0); }
+
+		TArray<FName> Registrations;
+	};
+
+	FLocalizableMessageProcessor();
+	~FLocalizableMessageProcessor();
+
+	LOCALIZABLEMESSAGE_API FText Localize(const FLocalizableMessage& Message, const FLocalizationContext& Context);
+
+	template <typename UserType>
+	void RegisterLocalizableType(const TFunction<FText(const UserType&, const FLocalizationContext&)>& LocalizeValueFunctor, FScopedRegistrations& ScopedRegistrations)
+	{
+		auto FncLocalizeValue = [LocalizeValueFunctor](const FLocalizableMessageParameter& Localizable, const FLocalizationContext& LocalizationContext)
+		{
+			return LocalizeValueFunctor(static_cast<const UserType&>(Localizable), LocalizationContext);
+		};
+
+		FName UserId = UserType::StaticStruct()->GetFName();
+		LocalizeValueMapping.Add(UserId, FncLocalizeValue);
+		ScopedRegistrations.Registrations.Add(UserId);
+	}
+
+	LOCALIZABLEMESSAGE_API void UnregisterLocalizableTypes(FScopedRegistrations& ScopedRegistrations);
+
+private:
+	using LocalizeValueFnc = TFunction<FText(const FLocalizableMessageParameter&, const FLocalizationContext&)>;
+
+	TMap<FName, LocalizeValueFnc> LocalizeValueMapping;
+};
\ No newline at end of file
diff --git a/Engine/Plugins/Experimental/LocalizableMessage/Source/LocalizableMessage/Public/LocalizableMessageTypes.h b/Engine/Plugins/Experimental/LocalizableMessage/Source/LocalizableMessage/Public/LocalizableMessageTypes.h
new file mode 100644
--- /dev/null
+++ b/Engine/Plugins/Experimental/LocalizableMessage/Source/LocalizableMessage/Public/LocalizableMessageTypes.h
@@ -0,0 +1,11 @@
+// Copyright Epic Games, Inc. All Rights Reserved.
+
+#pragma once
+
+#include "CoreMinimal.h"
+
+namespace LocalizableMessageTypes
+{
+	void RegisterTypes();
+	void UnregisterTypes();
+}
diff --git a/Engine/Plugins/Experimental/LocalizableMessage/Source/LocalizableMessage/Public/LocalizationContext.h b/Engine/Plugins/Experimental/LocalizableMessage/Source/LocalizableMessage/Public/LocalizationContext.h
new file mode 100644
--- /dev/null
+++ b/Engine/Plugins/Experimental/LocalizableMessage/Source/LocalizableMessage/Public/LocalizationContext.h
@@ -0,0 +1,15 @@
+// Copyright Epic Games, Inc. All Rights Reserved.
+
+#pragma once
+
+#include "CoreMinimal.h"
+#include "UObject/WeakObjectPtr.h"
+
+struct FLocalizationContext
+{
+	FLocalizationContext() = default;
+	explicit FLocalizationContext(UObject* InWorldContext) : WorldContext(InWorldContext) {}
+
+	TWeakObjectPtr<UObject> WorldContext = nullptr;
+};
+
