Add WithPureVirtual

diff --git a/Engine/Plugins/Experimental/LocalizableMessage/Source/LocalizableMessage/Public/LocalizableMessageParameter.h b/Engine/Plugins/Experimental/LocalizableMessage/Source/LocalizableMessage/Public/LocalizableMessageParameter.h
--- a/Engine/Plugins/Experimental/LocalizableMessage/Source/LocalizableMessage/Public/LocalizableMessageParameter.h
+++ b/Engine/Plugins/Experimental/LocalizableMessage/Source/LocalizableMessage/Public/LocalizableMessageParameter.h
@@ -28,4 +28,14 @@ struct FLocalizableMessageParameter
 	{
 		LOCALIZABLEMESSAGE_API void operator()(FLocalizableMessageParameter* Object) const;
 	};
-};
\ No newline at end of file
+};
+
+template<>
+struct TStructOpsTypeTraits< FLocalizableMessageParameter > : public TStructOpsTypeTraitsBase2< FLocalizableMessageParameter >
+{
+	enum
+	{
+		WithPureVirtual = true,
+	};
+};
+
