Use the new AutoRTFM instrumentation path which lets us cut down on the number of modules we need to compile with transactional support.

diff --git a/Engine/Plugins/Experimental/LocalizableMessage/Source/LocalizableMessage/LocalizableMessage.Build.cs b/Engine/Plugins/Experimental/LocalizableMessage/Source/LocalizableMessage/LocalizableMessage.Build.cs
--- a/Engine/Plugins/Experimental/LocalizableMessage/Source/LocalizableMessage/LocalizableMessage.Build.cs
+++ b/Engine/Plugins/Experimental/LocalizableMessage/Source/LocalizableMessage/LocalizableMessage.Build.cs
@@ -11,6 +11,8 @@ namespace UnrealBuildTool.Rules
 					"Core",
 					"CoreUObject",
 				});
+
+			bAllowAutoRTFMInstrumentation = true;
 		}
 	}
 }
diff --git a/Engine/Source/Programs/AutoRTFMTests/AutoRTFMTests.Build.cs b/Engine/Source/Programs/AutoRTFMTests/AutoRTFMTests.Build.cs
--- a/Engine/Source/Programs/AutoRTFMTests/AutoRTFMTests.Build.cs
+++ b/Engine/Source/Programs/AutoRTFMTests/AutoRTFMTests.Build.cs
@@ -23,5 +23,7 @@ public class AutoRTFMTests : ModuleRules
 				new string[] { "DesktopPlatform" }
 			);
 		}
+
+		bAllowAutoRTFMInstrumentation = true;
 	}
 }
