Add BP function to translate from a LocalizableMessage to FText.

diff --git a/Engine/Plugins/Experimental/LocalizableMessage/LocalizableMessage.uplugin b/Engine/Plugins/Experimental/LocalizableMessage/LocalizableMessage.uplugin
--- a/Engine/Plugins/Experimental/LocalizableMessage/LocalizableMessage.uplugin
+++ b/Engine/Plugins/Experimental/LocalizableMessage/LocalizableMessage.uplugin
@@ -19,6 +19,11 @@
 			"Name": "LocalizableMessage",
 			"Type": "Runtime",
 			"LoadingPhase": "Default"
+		},
+		{
+			"Name": "LocalizableMessageBlueprint",
+			"Type": "Runtime",
+			"LoadingPhase": "Default"
 		}
 	],
 	"Plugins": [
diff --git a/Engine/Plugins/Experimental/LocalizableMessage/Source/LocalizableMessageBlueprint/LocalizableMessageBlueprint.Build.cs b/Engine/Plugins/Experimental/LocalizableMessage/Source/LocalizableMessageBlueprint/LocalizableMessageBlueprint.Build.cs
new file mode 100644
--- /dev/null
+++ b/Engine/Plugins/Experimental/LocalizableMessage/Source/LocalizableMessageBlueprint/LocalizableMessageBlueprint.Build.cs
@@ -0,0 +1,19 @@
+// Copyright Epic Games, Inc. All Rights Reserved.
+
+namespace UnrealBuildTool.Rules
+{
+	public class LocalizableMessageBlueprint : ModuleRules
+	{
+		public LocalizableMessageBlueprint(ReadOnlyTargetRules Target) : base(Target)
+		{
+			PrivateDependencyModuleNames.AddRange(
+				new string[] {
+					"Core",
+					"CoreUObject",
+					"Engine",
+					"LocalizableMessage",
+					"StructUtils"
+				});
+		}
+	}
+}
diff --git a/Engine/Plugins/Experimental/LocalizableMessage/Source/LocalizableMessageBlueprint/Private/LocalizableMessageBlueprintModule.cpp b/Engine/Plugins/Experimental/LocalizableMessage/Source/LocalizableMessageBlueprint/Private/LocalizableMessageBlueprintModule.cpp
new file mode 100644
--- /dev/null
+++ b/Engine/Plugins/Experimental/LocalizableMessage/Source/LocalizableMessageBlueprint/Private/LocalizableMessageBlueprintModule.cpp
@@ -0,0 +1,5 @@
+// Copyright Epic Games, Inc. All Rights Reserved.
+
+#include "Modules/ModuleManager.h"
+
+IMPLEMENT_MODULE(FDefaultModuleImpl, LocalizableMessageBlueprint)
diff --git a/Engine/Plugins/Experimental/LocalizableMessage/Source/LocalizableMessageBlueprint/Private/LocalizableMessageLibrary.cpp b/Engine/Plugins/Experimental/LocalizableMessage/Source/LocalizableMessageBlueprint/Private/LocalizableMessageLibrary.cpp
new file mode 100644
--- /dev/null
+++ b/Engine/Plugins/Experimental/LocalizableMessage/Source/LocalizableMessageBlueprint/Private/LocalizableMessageLibrary.cpp
@@ -0,0 +1,44 @@
+// Copyright Epic Games, Inc. All Rights Reserved.
+
+#include "LocalizableMessageLibrary.h"
+#include "LocalizationContext.h"
+#include "LocalizableMessageProcessor.h"
+#include "ILocalizableMessageModule.h"
+#include "Internationalization/Internationalization.h"
+
+FText ULocalizableMessageLibrary::Conv_LocalizableMessageToText(const UObject* WorldContextObject, const int32& Message)
+{
+	// This will never be called, the exec version below will be hit instead
+	check(0);
+	return FText::GetEmpty();
+}
+
+DEFINE_FUNCTION(ULocalizableMessageLibrary::execConv_LocalizableMessageToText)
+{
+	P_GET_OBJECT(UObject, WorldContextObject)
+	Stack.StepCompiledIn<FStructProperty>(nullptr);
+	void* MessagePtr = Stack.MostRecentPropertyAddress;
+	FStructProperty* StructProp = CastField<FStructProperty>(Stack.MostRecentProperty);
+
+	P_FINISH;
+
+	FText Result = FText::GetEmpty();
+	if (StructProp && MessagePtr)
+	{
+		P_NATIVE_BEGIN
+		FLocalizationContext LocContext(WorldContextObject, FInternationalization::Get().GetCurrentCulture());
+		FLocalizableMessageProcessor& Processor = ILocalizableMessageModule::Get().GetLocalizableMessageProcessor();
+		Result = Processor.Localize(*reinterpret_cast<FLocalizableMessage*>(MessagePtr), LocContext);
+		P_NATIVE_END
+	}
+	else
+	{
+		FBlueprintExceptionInfo ExceptionInfo(
+			EBlueprintExceptionType::AccessViolation,
+			NSLOCTEXT("LocalizableMessageLibrary", "MissingMessageProperty", "Failed to resolve the input parameter for LocalizableMessageToText.")
+		);
+		FBlueprintCoreDelegates::ThrowScriptException(P_THIS, Stack, ExceptionInfo);
+	}
+
+	*(FText*)RESULT_PARAM = Result;
+}
diff --git a/Engine/Plugins/Experimental/LocalizableMessage/Source/LocalizableMessageBlueprint/Public/LocalizableMessageLibrary.h b/Engine/Plugins/Experimental/LocalizableMessage/Source/LocalizableMessageBlueprint/Public/LocalizableMessageLibrary.h
new file mode 100644
--- /dev/null
+++ b/Engine/Plugins/Experimental/LocalizableMessage/Source/LocalizableMessageBlueprint/Public/LocalizableMessageLibrary.h
@@ -0,0 +1,27 @@
+// Copyright Epic Games, Inc. All Rights Reserved.
+
+#pragma once
+
+#include "Kismet/BlueprintFunctionLibrary.h"
+#include "LocalizableMessage.h"
+
+#include "LocalizableMessageLibrary.generated.h"
+
+/** BlueprintFunctionLibrary for LocalizationMessage */
+UCLASS(Abstract, MinimalAPI, BlueprintInternalUseOnly)
+class ULocalizableMessageLibrary : public UBlueprintFunctionLibrary
+{
+	GENERATED_BODY()
+
+public:
+	/**
+	 * Conversion function from LocalizationMessage to FText.
+	 * @note Is only valid on the client. The processor generate a 
+	 */
+	UFUNCTION(BlueprintPure, BlueprintCosmetic, CustomThunk, Category = "Localization Message", meta = (CustomStructureParam = "Message", WorldContext = "WorldContextObject", BlueprintInternalUseOnly = "true"))
+	static LOCALIZABLEMESSAGEBLUEPRINT_API FText Conv_LocalizableMessageToText(const UObject* WorldContextObject, UPARAM(ref) const int32& Message);
+
+private:
+	DECLARE_FUNCTION(execConv_LocalizableMessageToText);
+};
+
