Make Localizable Message a blueprint type

diff --git a/Engine/Plugins/Experimental/LocalizableMessage/Source/LocalizableMessage/Public/LocalizableMessage.h b/Engine/Plugins/Experimental/LocalizableMessage/Source/LocalizableMessage/Public/LocalizableMessage.h
--- a/Engine/Plugins/Experimental/LocalizableMessage/Source/LocalizableMessage/Public/LocalizableMessage.h
+++ b/Engine/Plugins/Experimental/LocalizableMessage/Source/LocalizableMessage/Public/LocalizableMessage.h
@@ -29,7 +29,7 @@ public:
 	FInstancedStruct Value;
 };
 
-USTRUCT()
+USTRUCT(BlueprintType)
 struct FLocalizableMessage
 {
 	GENERATED_BODY();
