Updated FLocalizableMessage to handle localization/internationalization correctly

* FLocalizableMessageProcessor::Localize now uses the same method as NSLOCTEXT to find the localized FText at runtime.
* FLocalizationContext no longer provides a monolithic "culture" override, but rather two optional language and locale overrides.
    * This better matches how UE works internally (where language (for localization), and locale (for internationalization) are a split concept), and avoids baking a particular language/locale into the resultant format (which means it will save/load/replicate correctly when used as FText).

#rnx

diff --git a/Engine/Plugins/Experimental/LocalizableMessage/Source/LocalizableMessage/Private/LocalizableMessageProcessor.cpp b/Engine/Plugins/Experimental/LocalizableMessage/Source/LocalizableMessage/Private/LocalizableMessageProcessor.cpp
--- a/Engine/Plugins/Experimental/LocalizableMessage/Source/LocalizableMessage/Private/LocalizableMessageProcessor.cpp
+++ b/Engine/Plugins/Experimental/LocalizableMessage/Source/LocalizableMessage/Private/LocalizableMessageProcessor.cpp
@@ -18,23 +18,15 @@ FLocalizableMessageProcessor::~FLocalizableMessageProcessor()
 
 FText FLocalizableMessageProcessor::Localize(const FLocalizableMessage& Message, const FLocalizationContext& Context)
 {
-	// todo: look up the localized string via Key
-	// for now we only use the DefaultText in the message
-	FText SubstitutionResult;
 	FFormatNamedArguments FormatArguments;
-
 	for (const FLocalizableMessageParameterEntry& Substitution : Message.Substitutions)
 	{
 		if (ensure(Substitution.Value.IsValid()))
 		{
 			if (const LocalizeValueFnc* Functor = LocalizeValueMapping.Find(Substitution.Value.GetScriptStruct()->GetFName()))
 			{
-				SubstitutionResult = (*Functor)(Substitution.Value, Context);
-
-				if (SubstitutionResult.IsEmpty() == false)
-				{
-					TextFormatUtil::FormatNamed(FormatArguments, Substitution.Key, SubstitutionResult);
-				}
+				FText SubstitutionResult = (*Functor)(Substitution.Value, Context);
+				FormatArguments.Add(Substitution.Key, MoveTemp(SubstitutionResult));
 			}
 			else
 			{
@@ -48,7 +40,10 @@ FText FLocalizableMessageProcessor::Localize(const FLocalizableMessage& Message,
 	}
 
 	// an unfortunate number of allocations and copies here
-	return FText::Format(FTextFormat::FromString(Message.DefaultText), FormatArguments);
+	FText DefaultFText = FInternationalization::ForUseOnlyByLocMacroAndGraphNodeTextLiterals_CreateText(*Message.DefaultText, TEXT(""), *Message.Key);
+	return FormatArguments.Num() > 0
+		? FText::Format(DefaultFText, FormatArguments)
+		: DefaultFText;
 }
 
 void FLocalizableMessageProcessor::UnregisterLocalizableTypes(FScopedRegistrations& ScopedRegistrations)
diff --git a/Engine/Plugins/Experimental/LocalizableMessage/Source/LocalizableMessage/Private/LocalizableMessageTypes.cpp b/Engine/Plugins/Experimental/LocalizableMessage/Source/LocalizableMessage/Private/LocalizableMessageTypes.cpp
--- a/Engine/Plugins/Experimental/LocalizableMessage/Source/LocalizableMessage/Private/LocalizableMessageTypes.cpp
+++ b/Engine/Plugins/Experimental/LocalizableMessage/Source/LocalizableMessage/Private/LocalizableMessageTypes.cpp
@@ -13,15 +13,15 @@ namespace LocalizableMessageTypes
 
 	FText Int_LocalizeValue(const FLocalizableMessageParameterInt& Localizable, const FLocalizationContext& LocalizationContext)
 	{
-		return FText::AsNumber(Localizable.Value, nullptr, LocalizationContext.Culture);
+		return FText::AsNumber(Localizable.Value, nullptr, LocalizationContext.GetLocaleOverride());
 	}
 	FText Float_LocalizeValue(const FLocalizableMessageParameterFloat& Localizable, const FLocalizationContext& LocalizationContext)
 	{
-		return FText::AsNumber(Localizable.Value, nullptr, LocalizationContext.Culture);
+		return FText::AsNumber(Localizable.Value, nullptr, LocalizationContext.GetLocaleOverride());
 	}
 	FText String_LocalizeValue(const FLocalizableMessageParameterString& Localizable, const FLocalizationContext& LocalizationContext)
 	{
-		return FText::FromString(Localizable.Value);
+		return FText::AsCultureInvariant(Localizable.Value);
 	}
 
 	void RegisterTypes()
diff --git a/Engine/Plugins/Experimental/LocalizableMessage/Source/LocalizableMessage/Public/LocalizationContext.h b/Engine/Plugins/Experimental/LocalizableMessage/Source/LocalizableMessage/Public/LocalizationContext.h
--- a/Engine/Plugins/Experimental/LocalizableMessage/Source/LocalizableMessage/Public/LocalizationContext.h
+++ b/Engine/Plugins/Experimental/LocalizableMessage/Source/LocalizableMessage/Public/LocalizationContext.h
@@ -9,9 +9,38 @@
 struct FLocalizationContext
 {
 	FLocalizationContext() = default;
-	explicit FLocalizationContext(UObject* InWorldContext, const FCultureRef& InCulture) : WorldContext(InWorldContext), Culture(InCulture) {}
+	explicit FLocalizationContext(UObject* InWorldContext, FCulturePtr InLanguageOverride = nullptr, FCulturePtr InLocaleOverride = nullptr)
+		: WorldContext(InWorldContext)
+		, LanguageOverride(MoveTemp(InLanguageOverride))
+		, LocaleOverride(MoveTemp(InLocaleOverride))
+	{}
 
 	TWeakObjectPtr<UObject> WorldContext = nullptr;
-	FCultureRef Culture;
+
+	FCulturePtr GetLanguageOverride() const
+	{
+		return LanguageOverride;
+	}
+
+	FCulturePtr GetLocaleOverride() const
+	{
+		return LocaleOverride
+			? LocaleOverride
+			: LanguageOverride;
+	}
+
+private:
+	/**
+	 * The current language (for localization) override (if any).
+	 * @note Will also be used as the locale if LocaleOverride is null.
+	 * @note UE doesn't currently support loading the localization data for multiple languages at the same time.
+	 */
+	FCulturePtr LanguageOverride;
+
+	/**
+	 * The current locale (for internationalization) override (if any).
+	 * @note Internationalization affects things like number, date, time, etc formatting.
+	 */
+	FCulturePtr LocaleOverride;
 };
 
diff --git a/Engine/Plugins/Experimental/UIFramework/Source/Private/Widgets/UIFSimpleButton.cpp b/Engine/Plugins/Experimental/UIFramework/Source/Private/Widgets/UIFSimpleButton.cpp
--- a/Engine/Plugins/Experimental/UIFramework/Source/Private/Widgets/UIFSimpleButton.cpp
+++ b/Engine/Plugins/Experimental/UIFramework/Source/Private/Widgets/UIFSimpleButton.cpp
@@ -85,7 +85,7 @@ void UUIFrameworkSimpleButton::ServerClick_Implementation(APlayerController* Pla
 
 void UUIFrameworkSimpleButton::SetText(const FLocalizableMessage& InMessage)
 {
-	FLocalizationContext LocContext(this, FInternationalization::Get().GetCurrentCulture());
+	FLocalizationContext LocContext(this);
 	ILocalizableMessageModule& LocalizableMessageModule = ILocalizableMessageModule::Get();
 	FLocalizableMessageProcessor& Processor = LocalizableMessageModule.GetLocalizableMessageProcessor();
 	Text = Processor.Localize(InMessage, LocContext);
diff --git a/Engine/Plugins/Experimental/UIFramework/Source/Private/Widgets/UIFTextBlock.cpp b/Engine/Plugins/Experimental/UIFramework/Source/Private/Widgets/UIFTextBlock.cpp
--- a/Engine/Plugins/Experimental/UIFramework/Source/Private/Widgets/UIFTextBlock.cpp
+++ b/Engine/Plugins/Experimental/UIFramework/Source/Private/Widgets/UIFTextBlock.cpp
@@ -36,7 +36,7 @@ void UUIFrameworkTextBase::GetLifetimeReplicatedProps(TArray<FLifetimeProperty>&
 
 void UUIFrameworkTextBase::SetText(const FLocalizableMessage& InMessage)
 {
-	FLocalizationContext LocContext(this, FInternationalization::Get().GetCurrentCulture());
+	FLocalizationContext LocContext(this);
 	ILocalizableMessageModule& LocalizableMessageModule = ILocalizableMessageModule::Get();
 	FLocalizableMessageProcessor& Processor = LocalizableMessageModule.GetLocalizableMessageProcessor();
 	Text = Processor.Localize(InMessage, LocContext);
