Use FLocalizableMessage directly now that it's a blueprinttype

diff --git a/Engine/Plugins/Experimental/LocalizableMessage/Source/LocalizableMessageBlueprint/Private/LocalizableMessageLibrary.cpp b/Engine/Plugins/Experimental/LocalizableMessage/Source/LocalizableMessageBlueprint/Private/LocalizableMessageLibrary.cpp
--- a/Engine/Plugins/Experimental/LocalizableMessage/Source/LocalizableMessageBlueprint/Private/LocalizableMessageLibrary.cpp
+++ b/Engine/Plugins/Experimental/LocalizableMessage/Source/LocalizableMessageBlueprint/Private/LocalizableMessageLibrary.cpp
@@ -1,44 +1,16 @@
 // Copyright Epic Games, Inc. All Rights Reserved.
 
 #include "LocalizableMessageLibrary.h"
-#include "LocalizationContext.h"
-#include "LocalizableMessageProcessor.h"
+
 #include "ILocalizableMessageModule.h"
 #include "Internationalization/Internationalization.h"
+#include "LocalizableMessage.h"
+#include "LocalizableMessageProcessor.h"
+#include "LocalizationContext.h"
 
-FText ULocalizableMessageLibrary::Conv_LocalizableMessageToText(const UObject* WorldContextObject, const int32& Message)
-{
-	// This will never be called, the exec version below will be hit instead
-	check(0);
-	return FText::GetEmpty();
-}
-
-DEFINE_FUNCTION(ULocalizableMessageLibrary::execConv_LocalizableMessageToText)
+FText ULocalizableMessageLibrary::Conv_LocalizableMessageToText(UObject* WorldContextObject, const FLocalizableMessage& Message)
 {
-	P_GET_OBJECT(UObject, WorldContextObject)
-	Stack.StepCompiledIn<FStructProperty>(nullptr);
-	void* MessagePtr = Stack.MostRecentPropertyAddress;
-	FStructProperty* StructProp = CastField<FStructProperty>(Stack.MostRecentProperty);
-
-	P_FINISH;
-
-	FText Result = FText::GetEmpty();
-	if (StructProp && MessagePtr)
-	{
-		P_NATIVE_BEGIN
-		FLocalizationContext LocContext(WorldContextObject, FInternationalization::Get().GetCurrentCulture());
-		FLocalizableMessageProcessor& Processor = ILocalizableMessageModule::Get().GetLocalizableMessageProcessor();
-		Result = Processor.Localize(*reinterpret_cast<FLocalizableMessage*>(MessagePtr), LocContext);
-		P_NATIVE_END
-	}
-	else
-	{
-		FBlueprintExceptionInfo ExceptionInfo(
-			EBlueprintExceptionType::AccessViolation,
-			NSLOCTEXT("LocalizableMessageLibrary", "MissingMessageProperty", "Failed to resolve the input parameter for LocalizableMessageToText.")
-		);
-		FBlueprintCoreDelegates::ThrowScriptException(P_THIS, Stack, ExceptionInfo);
-	}
-
-	*(FText*)RESULT_PARAM = Result;
+	FLocalizationContext LocContext(WorldContextObject);
+	FLocalizableMessageProcessor& Processor = ILocalizableMessageModule::Get().GetLocalizableMessageProcessor();
+	return Processor.Localize(Message, LocContext);
 }
diff --git a/Engine/Plugins/Experimental/LocalizableMessage/Source/LocalizableMessageBlueprint/Public/LocalizableMessageLibrary.h b/Engine/Plugins/Experimental/LocalizableMessage/Source/LocalizableMessageBlueprint/Public/LocalizableMessageLibrary.h
--- a/Engine/Plugins/Experimental/LocalizableMessage/Source/LocalizableMessageBlueprint/Public/LocalizableMessageLibrary.h
+++ b/Engine/Plugins/Experimental/LocalizableMessage/Source/LocalizableMessageBlueprint/Public/LocalizableMessageLibrary.h
@@ -3,10 +3,12 @@
 #pragma once
 
 #include "Kismet/BlueprintFunctionLibrary.h"
-#include "LocalizableMessage.h"
 
 #include "LocalizableMessageLibrary.generated.h"
 
+struct FLocalizableMessage;
+class FText;
+
 /** BlueprintFunctionLibrary for LocalizationMessage */
 UCLASS(Abstract, MinimalAPI, BlueprintInternalUseOnly)
 class ULocalizableMessageLibrary : public UBlueprintFunctionLibrary
@@ -18,10 +20,7 @@ public:
 	 * Conversion function from LocalizationMessage to FText.
 	 * @note Is only valid on the client. The processor generate a 
 	 */
-	UFUNCTION(BlueprintPure, BlueprintCosmetic, CustomThunk, Category = "Localization Message", meta = (CustomStructureParam = "Message", WorldContext = "WorldContextObject", BlueprintInternalUseOnly = "true"))
-	static LOCALIZABLEMESSAGEBLUEPRINT_API FText Conv_LocalizableMessageToText(const UObject* WorldContextObject, UPARAM(ref) const int32& Message);
-
-private:
-	DECLARE_FUNCTION(execConv_LocalizableMessageToText);
+	UFUNCTION(BlueprintPure, BlueprintCosmetic, Category = "Localization Message", meta = (WorldContext = "WorldContextObject", BlueprintInternalUseOnly = "true"))
+	static LOCALIZABLEMESSAGEBLUEPRINT_API FText Conv_LocalizableMessageToText(UObject* WorldContextObject, UPARAM(ref) const FLocalizableMessage& Message);
 };
 
