Support message as a parameter for a message

diff --git a/Engine/Plugins/Experimental/LocalizableMessage/Source/LocalizableMessage/Private/LocalizableMessageTypes.cpp b/Engine/Plugins/Experimental/LocalizableMessage/Source/LocalizableMessage/Private/LocalizableMessageTypes.cpp
--- a/Engine/Plugins/Experimental/LocalizableMessage/Source/LocalizableMessage/Private/LocalizableMessageTypes.cpp
+++ b/Engine/Plugins/Experimental/LocalizableMessage/Source/LocalizableMessage/Private/LocalizableMessageTypes.cpp
@@ -23,6 +23,14 @@ namespace LocalizableMessageTypes
 	{
 		return FText::AsCultureInvariant(Localizable.Value);
 	}
+	FText Message_LocalizeValue(const FLocalizableMessageParameterMessage& Localizable, const FLocalizationContext& LocalizationContext)
+	{
+		ILocalizableMessageModule& LocalizableMessageModule = ILocalizableMessageModule::Get();
+		FLocalizableMessageProcessor& Processor = LocalizableMessageModule.GetLocalizableMessageProcessor();
+
+		return Processor.Localize(Localizable.Value, LocalizationContext);
+	}
+
 
 	void RegisterTypes()
 	{
@@ -32,6 +40,7 @@ namespace LocalizableMessageTypes
 		Processor.RegisterLocalizableType<FLocalizableMessageParameterInt>(&Int_LocalizeValue,RegisteredLocalizationTypes);
 		Processor.RegisterLocalizableType<FLocalizableMessageParameterFloat>(&Float_LocalizeValue,RegisteredLocalizationTypes);
 		Processor.RegisterLocalizableType<FLocalizableMessageParameterString>(&String_LocalizeValue, RegisteredLocalizationTypes);
+		Processor.RegisterLocalizableType<FLocalizableMessageParameterMessage>(&Message_LocalizeValue, RegisteredLocalizationTypes);
 	}
 
 	void UnregisterTypes()
diff --git a/Engine/Plugins/Experimental/LocalizableMessage/Source/LocalizableMessage/Public/LocalizableMessageBaseParameters.h b/Engine/Plugins/Experimental/LocalizableMessage/Source/LocalizableMessage/Public/LocalizableMessageBaseParameters.h
--- a/Engine/Plugins/Experimental/LocalizableMessage/Source/LocalizableMessage/Public/LocalizableMessageBaseParameters.h
+++ b/Engine/Plugins/Experimental/LocalizableMessage/Source/LocalizableMessage/Public/LocalizableMessageBaseParameters.h
@@ -3,6 +3,7 @@
 #pragma once
 
 #include "Containers/UnrealString.h"
+#include "LocalizableMessage.h"
 
 #include "LocalizableMessageBaseParameters.generated.h"
 
@@ -38,3 +39,14 @@ public:
 	UPROPERTY()
 	FString Value;
 };
+
+USTRUCT()
+struct FLocalizableMessageParameterMessage
+{
+	GENERATED_BODY()
+
+public:
+
+	UPROPERTY()
+	FLocalizableMessage Value;
+};
\ No newline at end of file
