Expose LocalizableMessageLibrary methods

#preglight 649079828d4eab0e55b9ab89

diff --git a/Engine/Plugins/Experimental/LocalizableMessage/Source/LocalizableMessage/Public/LocalizableMessage.h b/Engine/Plugins/Experimental/LocalizableMessage/Source/LocalizableMessage/Public/LocalizableMessage.h
--- a/Engine/Plugins/Experimental/LocalizableMessage/Source/LocalizableMessage/Public/LocalizableMessage.h
+++ b/Engine/Plugins/Experimental/LocalizableMessage/Source/LocalizableMessage/Public/LocalizableMessage.h
@@ -20,7 +20,10 @@ public:
 	LOCALIZABLEMESSAGE_API ~FLocalizableMessageParameterEntry();
 
 	bool operator==(const FLocalizableMessageParameterEntry& Other) const;
-	bool operator!=(const FLocalizableMessageParameterEntry& Other) const { return !(*this == Other); }
+	bool operator!=(const FLocalizableMessageParameterEntry& Other) const
+	{
+		return !(*this == Other);
+	}
 
 	UPROPERTY()
 	FString Key;
@@ -37,7 +40,23 @@ struct FLocalizableMessage
 	LOCALIZABLEMESSAGE_API FLocalizableMessage();
 	LOCALIZABLEMESSAGE_API ~FLocalizableMessage();
 	LOCALIZABLEMESSAGE_API bool operator==(const FLocalizableMessage& Other) const;
-	bool operator!=(const FLocalizableMessage& Other) const { return !(*this == Other); }
+
+	void Reset()
+	{
+		Key.Reset();
+		DefaultText.Reset();
+		Substitutions.Reset();
+	}
+
+	bool operator!=(const FLocalizableMessage& Other) const
+	{
+		return !(*this == Other);
+	}
+
+	bool IsEmpty() const
+	{
+		return Key.IsEmpty() && DefaultText.IsEmpty();
+	}
 
 	UPROPERTY()
 	FString Key;
diff --git a/Engine/Plugins/Experimental/LocalizableMessage/Source/LocalizableMessageBlueprint/Private/LocalizableMessageLibrary.cpp b/Engine/Plugins/Experimental/LocalizableMessage/Source/LocalizableMessageBlueprint/Private/LocalizableMessageLibrary.cpp
--- a/Engine/Plugins/Experimental/LocalizableMessage/Source/LocalizableMessageBlueprint/Private/LocalizableMessageLibrary.cpp
+++ b/Engine/Plugins/Experimental/LocalizableMessage/Source/LocalizableMessageBlueprint/Private/LocalizableMessageLibrary.cpp
@@ -14,3 +14,13 @@ FText ULocalizableMessageLibrary::Conv_LocalizableMessageToText(UObject* WorldCo
 	FLocalizableMessageProcessor& Processor = ILocalizableMessageModule::Get().GetLocalizableMessageProcessor();
 	return Processor.Localize(Message, LocContext);
 }
+
+bool ULocalizableMessageLibrary::IsEmpty_LocalizableMessage(const FLocalizableMessage& Message)
+{
+	return Message.IsEmpty();
+}
+
+void ULocalizableMessageLibrary::Reset_LocalizableMessage(FLocalizableMessage& Message)
+{
+	Message.Reset();
+}
\ No newline at end of file
diff --git a/Engine/Plugins/Experimental/LocalizableMessage/Source/LocalizableMessageBlueprint/Public/LocalizableMessageLibrary.h b/Engine/Plugins/Experimental/LocalizableMessage/Source/LocalizableMessageBlueprint/Public/LocalizableMessageLibrary.h
--- a/Engine/Plugins/Experimental/LocalizableMessage/Source/LocalizableMessageBlueprint/Public/LocalizableMessageLibrary.h
+++ b/Engine/Plugins/Experimental/LocalizableMessage/Source/LocalizableMessageBlueprint/Public/LocalizableMessageLibrary.h
@@ -9,18 +9,30 @@
 struct FLocalizableMessage;
 class FText;
 
-/** BlueprintFunctionLibrary for LocalizationMessage */
-UCLASS(Abstract, MinimalAPI, BlueprintInternalUseOnly)
+/** BlueprintFunctionLibrary for LocalizableMessage */
+UCLASS()
 class ULocalizableMessageLibrary : public UBlueprintFunctionLibrary
 {
 	GENERATED_BODY()
 
 public:
 	/**
-	 * Conversion function from LocalizationMessage to FText.
+	 * Conversion function from LocalizableMessage to FText.
 	 * @note Is only valid on the client. The processor generate a 
 	 */
-	UFUNCTION(BlueprintPure, BlueprintCosmetic, Category = "Localization Message", meta = (WorldContext = "WorldContextObject", BlueprintInternalUseOnly = "true"))
+	UFUNCTION(BlueprintPure, BlueprintCosmetic, Category = "Localizable Message", meta = (WorldContext = "WorldContextObject"))
 	static LOCALIZABLEMESSAGEBLUEPRINT_API FText Conv_LocalizableMessageToText(UObject* WorldContextObject, UPARAM(ref) const FLocalizableMessage& Message);
+
+	/**
+	 * Returns true if the message is empty
+	 */
+	UFUNCTION(BlueprintPure, Category = "Localizable Message")
+	static LOCALIZABLEMESSAGEBLUEPRINT_API bool IsEmpty_LocalizableMessage(UPARAM(ref) const FLocalizableMessage& Message);
+
+	/**
+	 * Resets the Localizable Message
+	 */
+	UFUNCTION(BlueprintCallable, Category = "Localizable Message")
+	static LOCALIZABLEMESSAGEBLUEPRINT_API void Reset_LocalizableMessage(UPARAM(ref) FLocalizableMessage& Message);
 };
 
