Minimal Export/Import class for ULocalizableMessageLibrary

diff --git a/Engine/Plugins/Experimental/LocalizableMessage/Source/LocalizableMessageBlueprint/Public/LocalizableMessageLibrary.h b/Engine/Plugins/Experimental/LocalizableMessage/Source/LocalizableMessageBlueprint/Public/LocalizableMessageLibrary.h
--- a/Engine/Plugins/Experimental/LocalizableMessage/Source/LocalizableMessageBlueprint/Public/LocalizableMessageLibrary.h
+++ b/Engine/Plugins/Experimental/LocalizableMessage/Source/LocalizableMessageBlueprint/Public/LocalizableMessageLibrary.h
@@ -10,7 +10,7 @@ struct FLocalizableMessage;
 class FText;
 
 /** BlueprintFunctionLibrary for LocalizableMessage */
-UCLASS()
+UCLASS(MinimalAPI)
 class ULocalizableMessageLibrary : public UBlueprintFunctionLibrary
 {
 	GENERATED_BODY()
