[ue] Unshelve from [REDACTED_HANDLE]
Skip FText::AsNumber calls in FLocalizableMessage where possible

diff --git a/Engine/Plugins/Experimental/LocalizableMessage/Source/LocalizableMessage/Private/LocalizableMessageProcessor.cpp b/Engine/Plugins/Experimental/LocalizableMessage/Source/LocalizableMessage/Private/LocalizableMessageProcessor.cpp
--- a/Engine/Plugins/Experimental/LocalizableMessage/Source/LocalizableMessage/Private/LocalizableMessageProcessor.cpp
+++ b/Engine/Plugins/Experimental/LocalizableMessage/Source/LocalizableMessage/Private/LocalizableMessageProcessor.cpp
@@ -25,7 +25,7 @@ FText FLocalizableMessageProcessor::Localize(const FLocalizableMessage& Message,
 		{
 			if (const LocalizeValueFnc* Functor = LocalizeValueMapping.Find(Substitution.Value.GetScriptStruct()->GetFName()))
 			{
-				FText SubstitutionResult = (*Functor)(Substitution.Value, Context);
+				FFormatArgumentValue SubstitutionResult = (*Functor)(Substitution.Value, Context);
 				FormatArguments.Add(Substitution.Key, MoveTemp(SubstitutionResult));
 			}
 			else
diff --git a/Engine/Plugins/Experimental/LocalizableMessage/Source/LocalizableMessage/Private/LocalizableMessageTypes.cpp b/Engine/Plugins/Experimental/LocalizableMessage/Source/LocalizableMessage/Private/LocalizableMessageTypes.cpp
--- a/Engine/Plugins/Experimental/LocalizableMessage/Source/LocalizableMessage/Private/LocalizableMessageTypes.cpp
+++ b/Engine/Plugins/Experimental/LocalizableMessage/Source/LocalizableMessage/Private/LocalizableMessageTypes.cpp
@@ -11,19 +11,25 @@ namespace LocalizableMessageTypes
 {
 	FLocalizableMessageProcessor::FScopedRegistrations RegisteredLocalizationTypes;
 
-	FText Int_LocalizeValue(const FLocalizableMessageParameterInt& Localizable, const FLocalizationContext& LocalizationContext)
+	FFormatArgumentValue Int_LocalizeValue(const FLocalizableMessageParameterInt& Localizable, const FLocalizationContext& LocalizationContext)
 	{
-		return FText::AsNumber(Localizable.Value, nullptr, LocalizationContext.GetLocaleOverride());
+		FCulturePtr LocaleOverride = LocalizationContext.GetLocaleOverride();
+		return LocaleOverride
+			? FFormatArgumentValue(FText::AsNumber(Localizable.Value, nullptr, LocalizationContext.GetLocaleOverride()))
+			: FFormatArgumentValue(Localizable.Value);
 	}
-	FText Float_LocalizeValue(const FLocalizableMessageParameterFloat& Localizable, const FLocalizationContext& LocalizationContext)
+	FFormatArgumentValue Float_LocalizeValue(const FLocalizableMessageParameterFloat& Localizable, const FLocalizationContext& LocalizationContext)
 	{
-		return FText::AsNumber(Localizable.Value, nullptr, LocalizationContext.GetLocaleOverride());
+		FCulturePtr LocaleOverride = LocalizationContext.GetLocaleOverride();
+		return LocaleOverride
+			? FFormatArgumentValue(FText::AsNumber(Localizable.Value, nullptr, LocalizationContext.GetLocaleOverride()))
+			: FFormatArgumentValue(Localizable.Value);
 	}
-	FText String_LocalizeValue(const FLocalizableMessageParameterString& Localizable, const FLocalizationContext& LocalizationContext)
+	FFormatArgumentValue String_LocalizeValue(const FLocalizableMessageParameterString& Localizable, const FLocalizationContext& LocalizationContext)
 	{
 		return FText::AsCultureInvariant(Localizable.Value);
 	}
-	FText Message_LocalizeValue(const FLocalizableMessageParameterMessage& Localizable, const FLocalizationContext& LocalizationContext)
+	FFormatArgumentValue Message_LocalizeValue(const FLocalizableMessageParameterMessage& Localizable, const FLocalizationContext& LocalizationContext)
 	{
 		ILocalizableMessageModule& LocalizableMessageModule = ILocalizableMessageModule::Get();
 		FLocalizableMessageProcessor& Processor = LocalizableMessageModule.GetLocalizableMessageProcessor();
diff --git a/Engine/Plugins/Experimental/LocalizableMessage/Source/LocalizableMessage/Public/LocalizableMessageProcessor.h b/Engine/Plugins/Experimental/LocalizableMessage/Source/LocalizableMessage/Public/LocalizableMessageProcessor.h
--- a/Engine/Plugins/Experimental/LocalizableMessage/Source/LocalizableMessage/Public/LocalizableMessageProcessor.h
+++ b/Engine/Plugins/Experimental/LocalizableMessage/Source/LocalizableMessage/Public/LocalizableMessageProcessor.h
@@ -9,6 +9,7 @@
 #include "UObject/NameTypes.h"
 
 class FText;
+class FFormatArgumentValue;
 struct FLocalizableMessage;
 struct FLocalizationContext;
 
@@ -30,7 +31,7 @@ public:
 	LOCALIZABLEMESSAGE_API FText Localize(const FLocalizableMessage& Message, const FLocalizationContext& Context);
 
 	template <typename UserType>
-	void RegisterLocalizableType(const TFunction<FText(const UserType&, const FLocalizationContext&)>& LocalizeValueFunctor, FScopedRegistrations& ScopedRegistrations)
+	void RegisterLocalizableType(const TFunction<FFormatArgumentValue(const UserType&, const FLocalizationContext&)>& LocalizeValueFunctor, FScopedRegistrations& ScopedRegistrations)
 	{
 		auto FncLocalizeValue = [LocalizeValueFunctor](const FInstancedStruct& Localizable, const FLocalizationContext& LocalizationContext)
 		{
@@ -45,7 +46,7 @@ public:
 	LOCALIZABLEMESSAGE_API void UnregisterLocalizableTypes(FScopedRegistrations& ScopedRegistrations);
 
 private:
-	using LocalizeValueFnc = TFunction<FText(const FInstancedStruct&, const FLocalizationContext&)>;
+	using LocalizeValueFnc = TFunction<FFormatArgumentValue(const FInstancedStruct&, const FLocalizationContext&)>;
 
 	TMap<FName, LocalizeValueFnc> LocalizeValueMapping;
 };
\ No newline at end of file
