Add a log when converting a localizablemessage to a text that displays the translated string.

diff --git a/Engine/Plugins/Experimental/LocalizableMessage/Source/LocalizableMessage/Private/LocalizableMessageProcessor.cpp b/Engine/Plugins/Experimental/LocalizableMessage/Source/LocalizableMessage/Private/LocalizableMessageProcessor.cpp
--- a/Engine/Plugins/Experimental/LocalizableMessage/Source/LocalizableMessage/Private/LocalizableMessageProcessor.cpp
+++ b/Engine/Plugins/Experimental/LocalizableMessage/Source/LocalizableMessage/Private/LocalizableMessageProcessor.cpp
@@ -6,6 +6,8 @@
 #include "LocalizableMessage.h"
 #include "LocalizationContext.h"
 
+DEFINE_LOG_CATEGORY(LogLocalizableMessageProcessor);
+
 FLocalizableMessageProcessor::FLocalizableMessageProcessor()
 {
 
@@ -41,9 +43,13 @@ FText FLocalizableMessageProcessor::Localize(const FLocalizableMessage& Message,
 
 	// an unfortunate number of allocations and copies here
 	FText DefaultFText = FInternationalization::ForUseOnlyByLocMacroAndGraphNodeTextLiterals_CreateText(*Message.DefaultText, TEXT(""), *Message.Key);
-	return FormatArguments.Num() > 0
+	FText RetFText = FormatArguments.Num() > 0
 		? FText::Format(DefaultFText, FormatArguments)
 		: DefaultFText;
+
+	UE_LOG(LogLocalizableMessageProcessor, VeryVerbose, TEXT("Localized Text: [%s]"),*RetFText.ToString());
+
+	return RetFText;
 }
 
 void FLocalizableMessageProcessor::UnregisterLocalizableTypes(FScopedRegistrations& ScopedRegistrations)
diff --git a/Engine/Plugins/Experimental/LocalizableMessage/Source/LocalizableMessage/Public/LocalizableMessageProcessor.h b/Engine/Plugins/Experimental/LocalizableMessage/Source/LocalizableMessage/Public/LocalizableMessageProcessor.h
--- a/Engine/Plugins/Experimental/LocalizableMessage/Source/LocalizableMessage/Public/LocalizableMessageProcessor.h
+++ b/Engine/Plugins/Experimental/LocalizableMessage/Source/LocalizableMessage/Public/LocalizableMessageProcessor.h
@@ -4,6 +4,7 @@
 
 #include "Containers/Map.h"
 #include "InstancedStruct.h"
+#include "Logging/LogMacros.h"
 #include "Templates/Function.h"
 #include "Templates/Tuple.h"
 #include "UObject/NameTypes.h"
@@ -49,4 +50,6 @@ private:
 	using LocalizeValueFnc = TFunction<FFormatArgumentValue(const FInstancedStruct&, const FLocalizationContext&)>;
 
 	TMap<FName, LocalizeValueFnc> LocalizeValueMapping;
-};
\ No newline at end of file
+};
+
+DECLARE_LOG_CATEGORY_EXTERN(LogLocalizableMessageProcessor, Log, All);
\ No newline at end of file
