Changed the function Conv_LocalizableMessageToText in LocalizableMessageLibrary to not be BlueprintCosmetic

#TechNotes BlueprintCosmetic was removed to be able to call this function from the server to get the message as an FText

diff --git a/Engine/Plugins/Experimental/LocalizableMessage/Source/LocalizableMessageBlueprint/Public/LocalizableMessageLibrary.h b/Engine/Plugins/Experimental/LocalizableMessage/Source/LocalizableMessageBlueprint/Public/LocalizableMessageLibrary.h
--- a/Engine/Plugins/Experimental/LocalizableMessage/Source/LocalizableMessageBlueprint/Public/LocalizableMessageLibrary.h
+++ b/Engine/Plugins/Experimental/LocalizableMessage/Source/LocalizableMessageBlueprint/Public/LocalizableMessageLibrary.h
@@ -18,9 +18,8 @@ class ULocalizableMessageLibrary : public UBlueprintFunctionLibrary
 public:
 	/**
 	 * Conversion function from LocalizableMessage to FText.
-	 * @note Is only valid on the client. The processor generate a 
 	 */
-	UFUNCTION(BlueprintPure, BlueprintCosmetic, Category = "Localizable Message", meta = (WorldContext = "WorldContextObject"))
+	UFUNCTION(BlueprintPure, Category = "Localizable Message", meta = (WorldContext = "WorldContextObject"))
 	static LOCALIZABLEMESSAGEBLUEPRINT_API FText Conv_LocalizableMessageToText(UObject* WorldContextObject, UPARAM(ref) const FLocalizableMessage& Message);
 
 	/**
