Open around getting the DefaultText as it goes through a cache that locks

diff --git a/Engine/Plugins/Experimental/LocalizableMessage/Source/LocalizableMessage/Private/LocalizableMessageProcessor.cpp b/Engine/Plugins/Experimental/LocalizableMessage/Source/LocalizableMessage/Private/LocalizableMessageProcessor.cpp
--- a/Engine/Plugins/Experimental/LocalizableMessage/Source/LocalizableMessage/Private/LocalizableMessageProcessor.cpp
+++ b/Engine/Plugins/Experimental/LocalizableMessage/Source/LocalizableMessage/Private/LocalizableMessageProcessor.cpp
@@ -42,7 +42,11 @@ FText FLocalizableMessageProcessor::Localize(const FLocalizableMessage& Message,
 	}
 
 	// an unfortunate number of allocations and copies here
-	FText DefaultFText = FInternationalization::ForUseOnlyByLocMacroAndGraphNodeTextLiterals_CreateText(*Message.DefaultText, TEXT(""), *Message.Key);
+	FText DefaultFText;
+	UE_AUTORTFM_OPEN({
+		DefaultFText = FInternationalization::ForUseOnlyByLocMacroAndGraphNodeTextLiterals_CreateText(*Message.DefaultText, TEXT(""), *Message.Key);
+	});
+
 	FText RetFText = FormatArguments.Num() > 0
 		? FText::Format(DefaultFText, FormatArguments)
 		: DefaultFText;
