Move dependency to public as it is included on headers.

diff --git a/Engine/Plugins/Experimental/LocalizableMessage/Source/LocalizableMessage/LocalizableMessage.Build.cs b/Engine/Plugins/Experimental/LocalizableMessage/Source/LocalizableMessage/LocalizableMessage.Build.cs
--- a/Engine/Plugins/Experimental/LocalizableMessage/Source/LocalizableMessage/LocalizableMessage.Build.cs
+++ b/Engine/Plugins/Experimental/LocalizableMessage/Source/LocalizableMessage/LocalizableMessage.Build.cs
@@ -10,7 +10,11 @@ namespace UnrealBuildTool.Rules
 				new string[] {
 					"Core",
 					"CoreUObject",
-					"StructUtils"
+				});
+
+			PublicDependencyModuleNames.AddRange(
+				new string[] {
+					"StructUtils" // used in header
 				});
 
 			bAllowAutoRTFMInstrumentation = true;
