Disable argument modifier evaluation for text created from message

#rnx

diff --git a/Engine/Plugins/Experimental/LocalizableMessage/Source/LocalizableMessage/Private/LocalizableMessageProcessor.cpp b/Engine/Plugins/Experimental/LocalizableMessage/Source/LocalizableMessage/Private/LocalizableMessageProcessor.cpp
--- a/Engine/Plugins/Experimental/LocalizableMessage/Source/LocalizableMessage/Private/LocalizableMessageProcessor.cpp
+++ b/Engine/Plugins/Experimental/LocalizableMessage/Source/LocalizableMessage/Private/LocalizableMessageProcessor.cpp
@@ -2,12 +2,29 @@
 
 #include "LocalizableMessageProcessor.h"
 
+#include "HAL/IConsoleManager.h"
 #include "Internationalization/Text.h"
 #include "LocalizableMessage.h"
 #include "LocalizationContext.h"
 
 DEFINE_LOG_CATEGORY(LogLocalizableMessageProcessor);
 
+namespace LocalizableMessageProcessor
+{
+
+// Disable argument modifier evaluation for text created from a message
+bool bAllowTextArgumentModifiers = false;
+FAutoConsoleVariableRef CVarAllowTextArgumentModifiers(TEXT("Localization.Message.AllowTextArgumentModifiers"), bAllowTextArgumentModifiers, TEXT("Whether to allow message -> text conversion to use text-style argument modifiers (default: false)"));
+
+ETextFormatFlags GetTextFormatFlags()
+{
+	return bAllowTextArgumentModifiers
+		? ETextFormatFlags::Default
+		: (ETextFormatFlags::Default & ~ETextFormatFlags::EvaluateArgumentModifiers);
+}
+
+}
+
 FLocalizableMessageProcessor::FLocalizableMessageProcessor()
 {
 
@@ -42,18 +59,19 @@ FText FLocalizableMessageProcessor::Localize(const FLocalizableMessage& Message,
 	}
 
 	// an unfortunate number of allocations and copies here
-	FText DefaultFText;
+	FText LocalizedText;
 	UE_AUTORTFM_OPEN({
-		DefaultFText = FInternationalization::ForUseOnlyByLocMacroAndGraphNodeTextLiterals_CreateText(*Message.DefaultText, TEXT(""), *Message.Key);
+		LocalizedText = FInternationalization::ForUseOnlyByLocMacroAndGraphNodeTextLiterals_CreateText(*Message.DefaultText, TEXT(""), *Message.Key);
 	});
+	if (FormatArguments.Num() > 0)
+	{
+		FTextFormat LocalizedTextFormat(LocalizedText, LocalizableMessageProcessor::GetTextFormatFlags());
+		LocalizedText = FText::Format(LocalizedTextFormat, FormatArguments);
+	}
 
-	FText RetFText = FormatArguments.Num() > 0
-		? FText::Format(DefaultFText, FormatArguments)
-		: DefaultFText;
-
-	UE_LOG(LogLocalizableMessageProcessor, VeryVerbose, TEXT("Localized Text: [%s]"),*RetFText.ToString());
+	UE_LOG(LogLocalizableMessageProcessor, VeryVerbose, TEXT("Localized Text: [%s]"), *LocalizedText.ToString());
 
-	return RetFText;
+	return LocalizedText;
 }
 
 void FLocalizableMessageProcessor::UnregisterLocalizableTypes(FScopedRegistrations& ScopedRegistrations)
