Only disable argument modifiers in FLocalizableMessage when processing source text

#rnx

#localization none
#tests local compilation

diff --git a/Engine/Plugins/Experimental/LocalizableMessage/Source/LocalizableMessage/Private/LocalizableMessageProcessor.cpp b/Engine/Plugins/Experimental/LocalizableMessage/Source/LocalizableMessage/Private/LocalizableMessageProcessor.cpp
--- a/Engine/Plugins/Experimental/LocalizableMessage/Source/LocalizableMessage/Private/LocalizableMessageProcessor.cpp
+++ b/Engine/Plugins/Experimental/LocalizableMessage/Source/LocalizableMessage/Private/LocalizableMessageProcessor.cpp
@@ -65,7 +65,8 @@ FText FLocalizableMessageProcessor::Localize(const FLocalizableMessage& Message,
 	});
 	if (FormatArguments.Num() > 0)
 	{
-		FTextFormat LocalizedTextFormat(LocalizedText, LocalizableMessageProcessor::GetTextFormatFlags());
+		const bool bIsLocalized = !LocalizedText.ToString().Equals(Message.DefaultText);
+		FTextFormat LocalizedTextFormat(LocalizedText, LocalizableMessageProcessor::GetTextFormatFlags() | (bIsLocalized ? ETextFormatFlags::EvaluateArgumentModifiers : ETextFormatFlags::None));
 		LocalizedText = FText::Format(LocalizedTextFormat, FormatArguments);
 	}
 
