Use const correctness for world context

diff --git a/Engine/Plugins/Experimental/LocalizableMessage/Source/LocalizableMessage/Public/LocalizationContext.h b/Engine/Plugins/Experimental/LocalizableMessage/Source/LocalizableMessage/Public/LocalizationContext.h
--- a/Engine/Plugins/Experimental/LocalizableMessage/Source/LocalizableMessage/Public/LocalizationContext.h
+++ b/Engine/Plugins/Experimental/LocalizableMessage/Source/LocalizableMessage/Public/LocalizationContext.h
@@ -9,13 +9,13 @@
 struct FLocalizationContext
 {
 	FLocalizationContext() = default;
-	explicit FLocalizationContext(UObject* InWorldContext, FCulturePtr InLanguageOverride = nullptr, FCulturePtr InLocaleOverride = nullptr)
+	explicit FLocalizationContext(const UObject* InWorldContext, FCulturePtr InLanguageOverride = nullptr, FCulturePtr InLocaleOverride = nullptr)
 		: WorldContext(InWorldContext)
 		, LanguageOverride(MoveTemp(InLanguageOverride))
 		, LocaleOverride(MoveTemp(InLocaleOverride))
 	{}
 
-	TWeakObjectPtr<UObject> WorldContext = nullptr;
+	TWeakObjectPtr<const UObject> WorldContext = nullptr;
 
 	FCulturePtr GetLanguageOverride() const
 	{
