Don't call ForUseOnlyByLocMacroAndGraphNodeTextLiterals_CreateText for an empty FLocalizableMessage

#rnx

diff --git a/Engine/Plugins/Experimental/LocalizableMessage/Source/LocalizableMessage/Private/LocalizableMessageProcessor.cpp b/Engine/Plugins/Experimental/LocalizableMessage/Source/LocalizableMessage/Private/LocalizableMessageProcessor.cpp
--- a/Engine/Plugins/Experimental/LocalizableMessage/Source/LocalizableMessage/Private/LocalizableMessageProcessor.cpp
+++ b/Engine/Plugins/Experimental/LocalizableMessage/Source/LocalizableMessage/Private/LocalizableMessageProcessor.cpp
@@ -60,9 +60,12 @@ FText FLocalizableMessageProcessor::Localize(const FLocalizableMessage& Message,
 
 	// an unfortunate number of allocations and copies here
 	FText LocalizedText;
-	UE_AUTORTFM_OPEN({
-		LocalizedText = FInternationalization::ForUseOnlyByLocMacroAndGraphNodeTextLiterals_CreateText(*Message.DefaultText, TEXT(""), *Message.Key);
-	});
+	if (!Message.IsEmpty())
+	{
+		UE_AUTORTFM_OPEN({
+			LocalizedText = FInternationalization::ForUseOnlyByLocMacroAndGraphNodeTextLiterals_CreateText(*Message.DefaultText, TEXT(""), *Message.Key);
+		});
+	}
 	if (FormatArguments.Num() > 0)
 	{
 		const bool bIsLocalized = !LocalizedText.ToString().Equals(Message.DefaultText);
