Enable AutoRTFM instrumention of three more modules.

diff --git a/Engine/Plugins/Experimental/LocalizableMessage/Source/LocalizableMessageBlueprint/LocalizableMessageBlueprint.Build.cs b/Engine/Plugins/Experimental/LocalizableMessage/Source/LocalizableMessageBlueprint/LocalizableMessageBlueprint.Build.cs
--- a/Engine/Plugins/Experimental/LocalizableMessage/Source/LocalizableMessageBlueprint/LocalizableMessageBlueprint.Build.cs
+++ b/Engine/Plugins/Experimental/LocalizableMessage/Source/LocalizableMessageBlueprint/LocalizableMessageBlueprint.Build.cs
@@ -14,6 +14,8 @@ namespace UnrealBuildTool.Rules
 					"LocalizableMessage",
 					"StructUtils"
 				});
+
+			bAllowAutoRTFMInstrumentation = true;
 		}
 	}
 }
