Remove AutoRTFM here

diff --git a/Engine/Plugins/Experimental/LocalizableMessage/Source/LocalizableMessageBlueprint/LocalizableMessageBlueprint.Build.cs b/Engine/Plugins/Experimental/LocalizableMessage/Source/LocalizableMessageBlueprint/LocalizableMessageBlueprint.Build.cs
--- a/Engine/Plugins/Experimental/LocalizableMessage/Source/LocalizableMessageBlueprint/LocalizableMessageBlueprint.Build.cs
+++ b/Engine/Plugins/Experimental/LocalizableMessage/Source/LocalizableMessageBlueprint/LocalizableMessageBlueprint.Build.cs
@@ -14,8 +14,6 @@ namespace UnrealBuildTool.Rules
 					"LocalizableMessage",
 					"StructUtils"
 				});
-
-			bAllowAutoRTFMInstrumentation = true;
 		}
 	}
 }
