Method to convert from FText to verse::message and back without losing namespace information

diff --git a/Engine/Plugins/Experimental/LocalizableMessage/Source/LocalizableMessage/Private/LocalizableMessage.cpp b/Engine/Plugins/Experimental/LocalizableMessage/Source/LocalizableMessage/Private/LocalizableMessage.cpp
--- a/Engine/Plugins/Experimental/LocalizableMessage/Source/LocalizableMessage/Private/LocalizableMessage.cpp
+++ b/Engine/Plugins/Experimental/LocalizableMessage/Source/LocalizableMessage/Private/LocalizableMessage.cpp
@@ -1,6 +1,68 @@
 // Copyright Epic Games, Inc. All Rights Reserved.
 
 #include "LocalizableMessage.h"
+#include "Internationalization/TextNamespaceUtil.h"
+
+namespace UE::LocalizableMessageTextInterop
+{
+	bool TextIdToMessageKey(const FTextId& TextId, FString& OutMessageKey)
+	{
+		if (TextId.IsEmpty())
+		{
+			return false;
+		}
+
+		FString Namespace = TextId.GetNamespace().ToString();
+#if USE_STABLE_LOCALIZATION_KEYS
+		// Remove any package ID from the namespace, as we only translate the "clean" version of the namespace
+		Namespace = TextNamespaceUtil::StripPackageNamespace(Namespace);
+#endif	// USE_STABLE_LOCALIZATION_KEYS
+
+		if (Namespace.IsEmpty())
+		{
+			// If there is no namespace then we can just use the key directly
+			TextId.GetKey().ToString(OutMessageKey);
+			return true;
+		}
+
+		// If there is a namespace then we need to combine the namespace+key together
+		Namespace.ReplaceInline(TEXT(","), TEXT("\\,"));
+
+		FString Key = TextId.GetKey().ToString();
+		Key.ReplaceInline(TEXT(","), TEXT("\\,"));
+
+		// Then we use a leading ~ character to denote this combination has happened
+		OutMessageKey = FString::Printf(TEXT("~%s,%s"), *Namespace, *Key);
+		return true;
+	}
+
+	bool MessageKeyToTextId(const FString& MessageKey, FTextId& OutTextId)
+	{
+		if (MessageKey.IsEmpty())
+		{
+			return false;
+		}
+
+		// This might be a combined namespace+key pair
+		if (MessageKey[0] == TEXT('~'))
+		{
+			FString Namespace;
+			FString Key;
+			if (MessageKey.Split(TEXT(","), &Namespace, &Key))
+			{
+				Namespace.RemoveAt(1, EAllowShrinking::No); // Remove the ~ marker denoting that this was a combined namespace+key pair
+				Namespace.ReplaceInline(TEXT("\\,"), TEXT(","));
+				Key.ReplaceInline(TEXT("\\,"), TEXT(","));
+				OutTextId = FTextId(Namespace, Key);
+				return true;
+			}
+			// Note: We don't return false if the Split failed, as it's possible (though unlikely) that this was a non-combined key that happened to start with a ~
+		}
+
+		OutTextId = FTextId(FTextKey(), MessageKey);
+		return true;
+	}
+}
 
 FLocalizableMessageParameterEntry::FLocalizableMessageParameterEntry() = default;
 FLocalizableMessageParameterEntry::FLocalizableMessageParameterEntry(const FString& InKey, const FInstancedStruct& InValue) : Key(InKey), Value(InValue) {}
@@ -20,4 +82,4 @@ bool FLocalizableMessage::operator==(const FLocalizableMessage& Other) const
 	return Key == Other.Key &&
 		DefaultText == Other.DefaultText &&
 		Substitutions == Other.Substitutions;
-}
\ No newline at end of file
+}
diff --git a/Engine/Plugins/Experimental/LocalizableMessage/Source/LocalizableMessage/Private/LocalizableMessageProcessor.cpp b/Engine/Plugins/Experimental/LocalizableMessage/Source/LocalizableMessage/Private/LocalizableMessageProcessor.cpp
--- a/Engine/Plugins/Experimental/LocalizableMessage/Source/LocalizableMessage/Private/LocalizableMessageProcessor.cpp
+++ b/Engine/Plugins/Experimental/LocalizableMessage/Source/LocalizableMessage/Private/LocalizableMessageProcessor.cpp
@@ -62,8 +62,13 @@ FText FLocalizableMessageProcessor::Localize(const FLocalizableMessage& Message,
 	FText LocalizedText;
 	if (!Message.IsEmpty())
 	{
-		UE_AUTORTFM_OPEN{ LocalizedText = FText::AsLocalizable_Advanced(FTextKey(), Message.Key, Message.DefaultText); };
+		FTextId TextId;
+		if (UE::LocalizableMessageTextInterop::MessageKeyToTextId(Message.Key, TextId))
+		{
+			UE_AUTORTFM_OPEN{ LocalizedText = FText::AsLocalizable_Advanced(TextId.GetNamespace(), TextId.GetKey(), Message.DefaultText); };
+		}
 	}
+
 	if (FormatArguments.Num() > 0)
 	{
 		const bool bIsLocalized = !LocalizedText.ToString().Equals(Message.DefaultText);
diff --git a/Engine/Plugins/Experimental/LocalizableMessage/Source/LocalizableMessage/Public/LocalizableMessage.h b/Engine/Plugins/Experimental/LocalizableMessage/Source/LocalizableMessage/Public/LocalizableMessage.h
--- a/Engine/Plugins/Experimental/LocalizableMessage/Source/LocalizableMessage/Public/LocalizableMessage.h
+++ b/Engine/Plugins/Experimental/LocalizableMessage/Source/LocalizableMessage/Public/LocalizableMessage.h
@@ -8,6 +8,22 @@
 
 #include "LocalizableMessage.generated.h"
 
+namespace UE::LocalizableMessageTextInterop
+{
+	/**
+	 * Attempt to convert the given FText ID to a string that can be used as a FLocalizableMessage key.
+	 * @return True if the conversion was possible, or false if it failed.
+	 */
+	LOCALIZABLEMESSAGE_API bool TextIdToMessageKey(const FTextId& TextId, FString& OutMessageKey);
+
+	/**
+	 * Attempt to convert the given FLocalizableMessage key into a FText ID.
+	 * @return True if the conversion was possible, or false if it failed.
+	 */
+	LOCALIZABLEMESSAGE_API bool MessageKeyToTextId(const FString& MessageKey, FTextId& OutTextId);
+}
+
+
 USTRUCT()
 struct LOCALIZABLEMESSAGE_API FLocalizableMessageParameterEntry
 {
