Verse message: helper function to convert a verse::message into FText.

#rnx

diff --git a/Engine/Plugins/Experimental/LocalizableMessage/Source/LocalizableMessage/Private/LocalizableMessageProcessor.cpp b/Engine/Plugins/Experimental/LocalizableMessage/Source/LocalizableMessage/Private/LocalizableMessageProcessor.cpp
--- a/Engine/Plugins/Experimental/LocalizableMessage/Source/LocalizableMessage/Private/LocalizableMessageProcessor.cpp
+++ b/Engine/Plugins/Experimental/LocalizableMessage/Source/LocalizableMessage/Private/LocalizableMessageProcessor.cpp
@@ -67,6 +67,10 @@ FText FLocalizableMessageProcessor::Localize(const FLocalizableMessage& Message,
 		{
 			LocalizedText = FText::AsLocalizable_Advanced(TextId.GetNamespace(), TextId.GetKey(), Message.DefaultText);
 		}
+		else
+		{
+			LocalizedText = FText::AsCultureInvariant(Message.DefaultText);
+		}
 	}
 
 	if (FormatArguments.Num() > 0)
