Made it so that Hud messages could be compared in Blueprints.

diff --git a/Engine/Plugins/Experimental/LocalizableMessage/Source/LocalizableMessageBlueprint/Private/LocalizableMessageLibrary.cpp b/Engine/Plugins/Experimental/LocalizableMessage/Source/LocalizableMessageBlueprint/Private/LocalizableMessageLibrary.cpp
--- a/Engine/Plugins/Experimental/LocalizableMessage/Source/LocalizableMessageBlueprint/Private/LocalizableMessageLibrary.cpp
+++ b/Engine/Plugins/Experimental/LocalizableMessage/Source/LocalizableMessageBlueprint/Private/LocalizableMessageLibrary.cpp
@@ -17,6 +17,11 @@ FText ULocalizableMessageLibrary::Conv_LocalizableMessageToText(UObject* WorldCo
 	return Processor.Localize(Message, LocContext);
 }
 
+bool ULocalizableMessageLibrary::EqualEqual_LocalizableMessage(const FLocalizableMessage& A, const FLocalizableMessage& B)
+{
+	return A == B;
+}
+
 bool ULocalizableMessageLibrary::IsEmpty_LocalizableMessage(const FLocalizableMessage& Message)
 {
 	return Message.IsEmpty();
@@ -25,4 +30,4 @@ bool ULocalizableMessageLibrary::IsEmpty_LocalizableMessage(const FLocalizableMe
 void ULocalizableMessageLibrary::Reset_LocalizableMessage(FLocalizableMessage& Message)
 {
 	Message.Reset();
-}
+}
\ No newline at end of file
diff --git a/Engine/Plugins/Experimental/LocalizableMessage/Source/LocalizableMessageBlueprint/Public/LocalizableMessageLibrary.h b/Engine/Plugins/Experimental/LocalizableMessage/Source/LocalizableMessageBlueprint/Public/LocalizableMessageLibrary.h
--- a/Engine/Plugins/Experimental/LocalizableMessage/Source/LocalizableMessageBlueprint/Public/LocalizableMessageLibrary.h
+++ b/Engine/Plugins/Experimental/LocalizableMessage/Source/LocalizableMessageBlueprint/Public/LocalizableMessageLibrary.h
@@ -22,10 +22,16 @@ public:
 	UFUNCTION(BlueprintPure, Category = "Localizable Message", meta = (WorldContext = "WorldContextObject"))
 	static LOCALIZABLEMESSAGEBLUEPRINT_API FText Conv_LocalizableMessageToText(UObject* WorldContextObject, UPARAM(ref) const FLocalizableMessage& Message);
 
+	/**
+	 * Returns true if the Localizable Messages are equal (A == B)
+	 */
+	UFUNCTION(BlueprintCallable, BlueprintPure, Category = "Localizable Message", meta = (DisplayName = "Equal (LocalizableMessage)", CompactNodeTitle = "==", BlueprintThreadSafe))
+	static LOCALIZABLEMESSAGEBLUEPRINT_API bool EqualEqual_LocalizableMessage(UPARAM(ref) const FLocalizableMessage& A, UPARAM(ref) const FLocalizableMessage& B);
+
 	/**
 	 * Returns true if the message is empty
 	 */
-	UFUNCTION(BlueprintPure, Category = "Localizable Message")
+	UFUNCTION(BlueprintCallable, BlueprintPure, Category = "Localizable Message")
 	static LOCALIZABLEMESSAGEBLUEPRINT_API bool IsEmpty_LocalizableMessage(UPARAM(ref) const FLocalizableMessage& Message);
 
 	/**
